PCG BiomeSample Plugin: Moved from Restricted/NotForLicensees to Experimental

#virtualized

diff --git a/Engine/Plugins/Experimental/PCGBiomeSample/PCGBiomeSample.uplugin b/Engine/Plugins/Experimental/PCGBiomeSample/PCGBiomeSample.uplugin
new file mode 100644
--- /dev/null
+++ b/Engine/Plugins/Experimental/PCGBiomeSample/PCGBiomeSample.uplugin
@@ -0,0 +1,24 @@
+{
+	"FileVersion": 3,
+	"Version": 1,
+	"VersionName": "0.1",
+	"FriendlyName": "PCG Biome Sample",
+	"Description": "PCG Biome Sample Plugin",
+	"Category": "Other",
+	"CreatedBy": "Epic Games, Inc.",
+	"CreatedByURL": "https://epicgames.com",
+	"DocsURL": "",
+	"MarketplaceURL": "",
+	"SupportURL": "",
+	"EnabledByDefault": false,
+	"CanContainContent": true,
+	"IsBetaVersion": false,
+	"IsExperimentalVersion": true,
+	"Installed": false,
+	"Plugins": [
+		{
+			"Name": "PCGBiomeCore",
+			"Enabled": true
+		}
+	]
+}
\ No newline at end of file
