PCG Biome Sample: added PCG plugin dependency to BiomeSample as well (was depending on BiomeCore before) to fix cook issues.

diff --git a/Engine/Plugins/Experimental/PCGBiomeSample/PCGBiomeSample.uplugin b/Engine/Plugins/Experimental/PCGBiomeSample/PCGBiomeSample.uplugin
--- a/Engine/Plugins/Experimental/PCGBiomeSample/PCGBiomeSample.uplugin
+++ b/Engine/Plugins/Experimental/PCGBiomeSample/PCGBiomeSample.uplugin
@@ -19,6 +19,10 @@
 		{
 			"Name": "PCGBiomeCore",
 			"Enabled": true
+		},
+		{
+			"Name": "PCG",
+			"Enabled": true
 		}
 	]
 }
\ No newline at end of file
